Fixed Win32 builds not compiling

diff --git a/Engine/Plugins/Runtime/TwitchLiveStreaming/Source/TwitchLiveStreaming/TwitchLiveStreaming.cpp b/Engine/Plugins/Runtime/TwitchLiveStreaming/Source/TwitchLiveStreaming/TwitchLiveStreaming.cpp
--- a/Engine/Plugins/Runtime/TwitchLiveStreaming/Source/TwitchLiveStreaming/TwitchLiveStreaming.cpp
+++ b/Engine/Plugins/Runtime/TwitchLiveStreaming/Source/TwitchLiveStreaming/TwitchLiveStreaming.cpp
@@ -374,7 +374,7 @@ void FTwitchLiveStreaming::InitTwitch( const FString& TwitchWindowsDLLFolder )
 	{
 		struct Local
 		{
-			static void* TwitchAlloc( SIZE_T Size, SIZE_T Alignment )
+			static void* TwitchAlloc( size_t Size, size_t Alignment )
 			{
 				void* Result = FMemory::Malloc( Size, uint32( Alignment ) );
 #if STATS
