Fixed regression with commandlets crashing on startup
Also, don't bother loading LiveStreaming support in commandlets

diff --git a/Engine/Plugins/Runtime/TwitchLiveStreaming/Source/TwitchLiveStreaming/TwitchLiveStreaming.cpp b/Engine/Plugins/Runtime/TwitchLiveStreaming/Source/TwitchLiveStreaming/TwitchLiveStreaming.cpp
--- a/Engine/Plugins/Runtime/TwitchLiveStreaming/Source/TwitchLiveStreaming/TwitchLiveStreaming.cpp
+++ b/Engine/Plugins/Runtime/TwitchLiveStreaming/Source/TwitchLiveStreaming/TwitchLiveStreaming.cpp
@@ -117,11 +117,15 @@ void FTwitchLiveStreaming::StartupModule()
 	}
 
 	// Register our custom project settings
-	ISettingsModule::Get()->RegisterSettings(
-		"Project", "Plugins", "Twitch",
-		LOCTEXT( "TwitchSettingsName", "Twitch" ),
-		LOCTEXT( "TwitchSettingsDescription", "Twitch game broadcasting settings" ),
-		GetMutableDefault< UTwitchProjectSettings >() );
+	ISettingsModule* SettingsModule = ISettingsModule::Get();
+	if( SettingsModule != nullptr )
+	{
+		SettingsModule->RegisterSettings(
+			"Project", "Plugins", "Twitch",
+			LOCTEXT( "TwitchSettingsName", "Twitch" ),
+			LOCTEXT( "TwitchSettingsDescription", "Twitch game broadcasting settings" ),
+			GetMutableDefault< UTwitchProjectSettings >() );
+	}
 
 	// Register as a modular feature
 	IModularFeatures::Get().RegisterModularFeature( TEXT( "LiveStreaming" ), this );
