Initialize Twitch on demand instead of at startup

diff --git a/Engine/Plugins/Runtime/TwitchLiveStreaming/Source/TwitchLiveStreaming/TwitchLiveStreaming.cpp b/Engine/Plugins/Runtime/TwitchLiveStreaming/Source/TwitchLiveStreaming/TwitchLiveStreaming.cpp
--- a/Engine/Plugins/Runtime/TwitchLiveStreaming/Source/TwitchLiveStreaming/TwitchLiveStreaming.cpp
+++ b/Engine/Plugins/Runtime/TwitchLiveStreaming/Source/TwitchLiveStreaming/TwitchLiveStreaming.cpp
@@ -92,30 +92,6 @@ void FTwitchLiveStreaming::StartupModule()
 	// @todo twitch: Editor: Show audio microphone input level, so you can be sure you are not muted
 	// @todo twitch: Editor: Ideally we want a "zoom in" feature to use during editor live streams
 
-	// Load the Twitch SDK dll
-	// @todo twitch: This all needs to be handled differently for iOS (static linkage)
-	FString TwitchDLLFolder( FPaths::Combine( 
-		*FPaths::EngineDir(), 
-		TEXT("Binaries/ThirdParty/NotForLicensees/Twitch/Twitch-6.17/"),	// Check the "NotForLicensees" folder first
-		FPlatformProcess::GetBinariesSubdirectory() ) );
-	if( !IFileManager::Get().DirectoryExists( *TwitchDLLFolder ) )
-	{
-		TwitchDLLFolder = FPaths::Combine( 
-			*FPaths::EngineDir(), 
-			TEXT("Binaries/ThirdParty/Twitch/Twitch-6.17/"),	// Check the normal folder
-			FPlatformProcess::GetBinariesSubdirectory() );
-	}
-
-
-	// Load the SDK DLL
-	LoadTwitchSDK( TwitchDLLFolder );
-
-	if( TwitchState == ETwitchState::DLLLoaded )
-	{
-		// Initialize the Twitch SDK
-		InitTwitch( TwitchDLLFolder );
-	}
-
 	// Register our custom project settings
 	ISettingsModule* SettingsModule = ISettingsModule::Get();
 	if( SettingsModule != nullptr )
@@ -137,10 +113,6 @@ void FTwitchLiveStreaming::ShutdownModule()
 	// Unregister our feature
 	IModularFeatures::Get().UnregisterModularFeature( TEXT( "LiveStreaming" ), this );
 
-	// No longer safe to run callback functions, as we could be shutting down
-	OnStatusChangedEvent.Clear();
-	OnChatMessageEvent.Clear();
-
 	// Unregister custom project settings
 	ISettingsModule* SettingsModule = ISettingsModule::Get();
 	if( SettingsModule != nullptr )
@@ -148,75 +120,113 @@ void FTwitchLiveStreaming::ShutdownModule()
 		SettingsModule->UnregisterSettings( "Project", "Plugins", "Twitch" );
 	}
 
-	// Turn off Twitch chat system
-	if( ChatState != EChatState::Uninitialized )
+	if( TwitchState != ETwitchState::Uninitialized )
 	{
-		UE_LOG( LogTwitch, Display, TEXT( "Shutting down Twitch chat system" ) );
-		const TTV_ErrorCode TwitchErrorCode = TwitchChatShutdown( nullptr, nullptr );
-		if( !TTV_SUCCEEDED( TwitchErrorCode ) )
+		// No longer safe to run callback functions, as we could be shutting down
+		OnStatusChangedEvent.Clear();
+		OnChatMessageEvent.Clear();
+
+		// Turn off Twitch chat system
+		if( ChatState != EChatState::Uninitialized )
 		{
-			const FString TwitchErrorString( UTF8_TO_TCHAR( TwitchErrorToString( TwitchErrorCode ) ) );
-			UE_LOG( LogTwitch, Warning, TEXT( "An error occured while shutting down Twitch's chat system.\n\nError: %s (%d)" ), *TwitchErrorString, (int32)TwitchErrorCode );
+			UE_LOG( LogTwitch, Display, TEXT( "Shutting down Twitch chat system" ) );
+			const TTV_ErrorCode TwitchErrorCode = TwitchChatShutdown( nullptr, nullptr );
+			if( !TTV_SUCCEEDED( TwitchErrorCode ) )
+			{
+				const FString TwitchErrorString( UTF8_TO_TCHAR( TwitchErrorToString( TwitchErrorCode ) ) );
+				UE_LOG( LogTwitch, Warning, TEXT( "An error occured while shutting down Twitch's chat system.\n\nError: %s (%d)" ), *TwitchErrorString, (int32)TwitchErrorCode );
+			}
+			ChatState = EChatState::Uninitialized;
 		}
-		ChatState = EChatState::Uninitialized;
-	}
 
-	// Turn off Twitch web cam system
-	if( WebCamState != EWebCamState::Uninitialized )
-	{
-		UE_LOG( LogTwitch, Display, TEXT( "Shutting down Twitch web cam system") );
-		const TTV_ErrorCode TwitchErrorCode = TwitchWebCamShutdown( nullptr, nullptr );
-		if( !TTV_SUCCEEDED( TwitchErrorCode ) )
+		// Turn off Twitch web cam system
+		if( WebCamState != EWebCamState::Uninitialized )
 		{
-			const FString TwitchErrorString( UTF8_TO_TCHAR( TwitchErrorToString( TwitchErrorCode ) ) );
-			UE_LOG( LogTwitch, Warning, TEXT( "An error occured while shutting down Twitch's web cam system.\n\nError: %s (%d)"), *TwitchErrorString, (int32)TwitchErrorCode );
+			UE_LOG( LogTwitch, Display, TEXT( "Shutting down Twitch web cam system") );
+			const TTV_ErrorCode TwitchErrorCode = TwitchWebCamShutdown( nullptr, nullptr );
+			if( !TTV_SUCCEEDED( TwitchErrorCode ) )
+			{
+				const FString TwitchErrorString( UTF8_TO_TCHAR( TwitchErrorToString( TwitchErrorCode ) ) );
+				UE_LOG( LogTwitch, Warning, TEXT( "An error occured while shutting down Twitch's web cam system.\n\nError: %s (%d)"), *TwitchErrorString, (int32)TwitchErrorCode );
+			}
+			WebCamState = EWebCamState::Uninitialized;
 		}
-		WebCamState = EWebCamState::Uninitialized;
-	}
-	
-	// Release the web cam texture
-	const bool bReleaseResourceToo = false;	// No need to release on shutdown, it will free up during normal GC phase
-	ReleaseWebCamTexture( bReleaseResourceToo );
 
-	// We are no longer broadcasting
-	if( BroadcastState != EBroadcastState::Idle )
-	{
-		// Broadcast will be forcibly stopped by Twitch when we shutdown below.  We'll release our graphics resources afterwards.
-		BroadcastState = EBroadcastState::Idle;
-	}
+		// Release the web cam texture
+		const bool bReleaseResourceToo = false;	// No need to release on shutdown, it will free up during normal GC phase
+		ReleaseWebCamTexture( bReleaseResourceToo );
 
-	if( TwitchState != ETwitchState::Uninitialized && TwitchState != ETwitchState::DLLLoaded )
-	{
-		// Turn off Twitch
-		UE_LOG( LogTwitch, Display, TEXT( "Shutting down Twitch SDK") );
-		const TTV_ErrorCode TwitchErrorCode = TwitchShutdown();
-		if( !TTV_SUCCEEDED( TwitchErrorCode ) )
+		// We are no longer broadcasting
+		if( BroadcastState != EBroadcastState::Idle )
 		{
-			const FString TwitchErrorString( UTF8_TO_TCHAR( TwitchErrorToString( TwitchErrorCode ) ) );
-			UE_LOG( LogTwitch, Warning, TEXT( "An error occured while shutting down Twitch.\n\nError: %s (%d)"), *TwitchErrorString, (int32)TwitchErrorCode );
+			// Broadcast will be forcibly stopped by Twitch when we shutdown below.  We'll release our graphics resources afterwards.
+			BroadcastState = EBroadcastState::Idle;
 		}
 
-		TwitchState = ETwitchState::DLLLoaded;
-	}
+		if( TwitchState != ETwitchState::Uninitialized && TwitchState != ETwitchState::DLLLoaded )
+		{
+			// Turn off Twitch
+			UE_LOG( LogTwitch, Display, TEXT( "Shutting down Twitch SDK") );
+			const TTV_ErrorCode TwitchErrorCode = TwitchShutdown();
+			if( !TTV_SUCCEEDED( TwitchErrorCode ) )
+			{
+				const FString TwitchErrorString( UTF8_TO_TCHAR( TwitchErrorToString( TwitchErrorCode ) ) );
+				UE_LOG( LogTwitch, Warning, TEXT( "An error occured while shutting down Twitch.\n\nError: %s (%d)"), *TwitchErrorString, (int32)TwitchErrorCode );
+			}
 
-	// Clean up video buffers
-	for( uint32 VideoBufferIndex = 0; VideoBufferIndex < TwitchVideoBufferCount; ++VideoBufferIndex )
-	{
-		delete[] VideoBuffers[ VideoBufferIndex ];
-		VideoBuffers[ VideoBufferIndex ] = nullptr;
+			TwitchState = ETwitchState::DLLLoaded;
+		}
+
+		// Clean up video buffers
+		for( uint32 VideoBufferIndex = 0; VideoBufferIndex < TwitchVideoBufferCount; ++VideoBufferIndex )
+		{
+			delete[] VideoBuffers[ VideoBufferIndex ];
+			VideoBuffers[ VideoBufferIndex ] = nullptr;
+		}
+		AvailableVideoBuffers.Reset();
+
+		if( TwitchState == ETwitchState::DLLLoaded )
+		{
+			check( TwitchDLLHandle != nullptr );
+
+			// Release the DLL
+			UE_LOG( LogTwitch, Display, TEXT( "Unloading Twitch SDK DLL") );
+			FPlatformProcess::FreeDllHandle( TwitchDLLHandle );
+			TwitchDLLHandle = nullptr;
+		}
+
+		TwitchState = ETwitchState::Uninitialized;
 	}
-	AvailableVideoBuffers.Reset();
+}
 
-	if( TwitchState == ETwitchState::DLLLoaded )
+
+void FTwitchLiveStreaming::InitOnDemand()
+{
+	if( TwitchState == ETwitchState::Uninitialized )
 	{
-		check( TwitchDLLHandle != nullptr );
+		// Load the Twitch SDK dll
+		// @todo twitch: This all needs to be handled differently for iOS (static linkage)
+		FString TwitchDLLFolder( FPaths::Combine( 
+			*FPaths::EngineDir(), 
+			TEXT("Binaries/ThirdParty/NotForLicensees/Twitch/Twitch-6.17/"),	// Check the "NotForLicensees" folder first
+			FPlatformProcess::GetBinariesSubdirectory() ) );
+		if( !IFileManager::Get().DirectoryExists( *TwitchDLLFolder ) )
+		{
+			TwitchDLLFolder = FPaths::Combine( 
+				*FPaths::EngineDir(), 
+				TEXT("Binaries/ThirdParty/Twitch/Twitch-6.17/"),	// Check the normal folder
+				FPlatformProcess::GetBinariesSubdirectory() );
+		}
 
-		// Release the DLL
-		UE_LOG( LogTwitch, Display, TEXT( "Unloading Twitch SDK DLL") );
-		FPlatformProcess::FreeDllHandle( TwitchDLLHandle );
-		TwitchDLLHandle = nullptr;
 
-		TwitchState = ETwitchState::Uninitialized;
+		// Load the SDK DLL
+		LoadTwitchSDK( TwitchDLLFolder );
+
+		if( TwitchState == ETwitchState::DLLLoaded )
+		{
+			// Initialize the Twitch SDK
+			InitTwitch( TwitchDLLFolder );
+		}
 	}
 }
 
@@ -983,6 +993,8 @@ ILiveStreamingService::FOnStatusChanged& FTwitchLiveStreaming::OnStatusChanged()
 
 void FTwitchLiveStreaming::StartBroadcasting( const FBroadcastConfig& Config )
 {
+	InitOnDemand();
+
 	if( !IsBroadcasting() )
 	{
 		bWantsToBroadcastNow = true;
@@ -996,6 +1008,8 @@ void FTwitchLiveStreaming::StartBroadcasting( const FBroadcastConfig& Config )
 
 void FTwitchLiveStreaming::StopBroadcasting()
 {
+	InitOnDemand();
+
 	if( IsBroadcasting() )
 	{
 		bWantsToBroadcastNow = false;
@@ -1065,6 +1079,7 @@ void FTwitchLiveStreaming::PushVideoFrame( const FColor* VideoFrameBuffer )
 
 void FTwitchLiveStreaming::StartWebCam( const FWebCamConfig& Config )
 {
+	InitOnDemand();
 	if( !IsWebCamEnabled() )
 	{
 		bWantsWebCamNow = true;
@@ -1078,6 +1093,7 @@ void FTwitchLiveStreaming::StartWebCam( const FWebCamConfig& Config )
 
 void FTwitchLiveStreaming::StopWebCam()
 {
+	InitOnDemand();
 	if( IsWebCamEnabled() )
 	{
 		bWantsWebCamNow = false;
@@ -1108,6 +1124,7 @@ ILiveStreamingService::FOnChatMessage& FTwitchLiveStreaming::OnChatMessage()
 
 void FTwitchLiveStreaming::ConnectToChat()
 {
+	InitOnDemand();
 	if( !IsChatEnabled() )
 	{
 		bWantsChatEnabled = true;
@@ -1120,6 +1137,7 @@ void FTwitchLiveStreaming::ConnectToChat()
 
 void FTwitchLiveStreaming::DisconnectFromChat()
 {
+	InitOnDemand();
 	if( IsChatEnabled() )
 	{
 		bWantsChatEnabled = false;
@@ -1147,6 +1165,8 @@ void FTwitchLiveStreaming::SendChatMessage( const FString& ChatMessage )
 
 void FTwitchLiveStreaming::QueryLiveStreams( const FString& GameName, FQueryLiveStreamsCallback CompletionCallback ) 
 {
+	InitOnDemand();
+
 	// @todo twitch urgent: Don't we need to be logged in first?  Need to do async login as part of this?  Or expose a Login API with status events
 	Async_GetGameLiveStreams( GameName, CompletionCallback );
 }
diff --git a/Engine/Plugins/Runtime/TwitchLiveStreaming/Source/TwitchLiveStreaming/TwitchLiveStreaming.h b/Engine/Plugins/Runtime/TwitchLiveStreaming/Source/TwitchLiveStreaming/TwitchLiveStreaming.h
--- a/Engine/Plugins/Runtime/TwitchLiveStreaming/Source/TwitchLiveStreaming/TwitchLiveStreaming.h
+++ b/Engine/Plugins/Runtime/TwitchLiveStreaming/Source/TwitchLiveStreaming/TwitchLiveStreaming.h
@@ -45,6 +45,9 @@ public:
 
 private:
 
+	/** Initializes Twitch lazily */
+	void InitOnDemand();
+
 	/** Loads the Twitch SDK dynamic library */
 	void LoadTwitchSDK( const FString& TwitchDLLFolder );
 
