Twitch authorization now lands on temp web site
- This prevents a 404 error which can interfere with our scraping

diff --git a/Engine/Plugins/Runtime/TwitchLiveStreaming/Source/TwitchLiveStreaming/TwitchLiveStreaming.cpp b/Engine/Plugins/Runtime/TwitchLiveStreaming/Source/TwitchLiveStreaming/TwitchLiveStreaming.cpp
--- a/Engine/Plugins/Runtime/TwitchLiveStreaming/Source/TwitchLiveStreaming/TwitchLiveStreaming.cpp
+++ b/Engine/Plugins/Runtime/TwitchLiveStreaming/Source/TwitchLiveStreaming/TwitchLiveStreaming.cpp
@@ -345,7 +345,7 @@ FTwitchLiveStreaming::FTwitchProjectSettings FTwitchLiveStreaming::LoadProjectSe
 		// usage of the editor to their Twitch channel.  This is only used when running the editor though.  To allow your users to
 		// stream gameplay to Twitch, you'll need to create your own Twitch client ID and set that in your project's settings!
 		Settings.ClientID = TEXT( "nu2tppawhw58o74a2burkyc9fyuvjlr" );
-		Settings.RedirectURI = TEXT( "http://localhost:6180" );	// @todo twitch: Need nicer redirect landing page
+		Settings.RedirectURI = TEXT( "http://unrealengine.com" );	// @todo twitch: Need nicer redirect landing page
 
 		// UserName, Password and ClientSecret are used for editor direct authentication with Twitch.  This authentication method 
 		// is less secure, and requires Twitch to whitelist your application's ClientID in order for it to authenticate.  When
