Temporarily disabled Twitch support in release builds (pending clear by legal)

diff --git a/Engine/Plugins/Runtime/TwitchLiveStreaming/Source/TwitchLiveStreaming/TwitchLiveStreaming.Build.cs b/Engine/Plugins/Runtime/TwitchLiveStreaming/Source/TwitchLiveStreaming/TwitchLiveStreaming.Build.cs
--- a/Engine/Plugins/Runtime/TwitchLiveStreaming/Source/TwitchLiveStreaming/TwitchLiveStreaming.Build.cs
+++ b/Engine/Plugins/Runtime/TwitchLiveStreaming/Source/TwitchLiveStreaming/TwitchLiveStreaming.Build.cs
@@ -29,6 +29,7 @@ namespace UnrealBuildTool.Rules
 				});
 
  			bool bHaveTwitchSDK =
+				!UnrealBuildTool.BuildingRocket() &&	// @todo: Twitch support disabled in released builds until it is cleared by legal
  				( System.IO.Directory.Exists( System.IO.Path.Combine( UEBuildConfiguration.UEThirdPartySourceDirectory, "Twitch" ) ) &&
  				  System.IO.Directory.Exists( System.IO.Path.Combine( UEBuildConfiguration.UEThirdPartySourceDirectory, "Twitch", "Twitch-6.17" ) ) ) ||
  				( System.IO.Directory.Exists( System.IO.Path.Combine( UEBuildConfiguration.UEThirdPartySourceDirectory, "NotForLicensees" ) ) &&
