Core - UE/UE-6813 - Confusing API: We have both FMemory::MemZero() and FMemory:Memzero()

diff --git a/Engine/Plugins/Editor/SpeedTreeImporter/Source/SpeedTreeImporter/Private/SpeedTreeImportFactory.cpp b/Engine/Plugins/Editor/SpeedTreeImporter/Source/SpeedTreeImporter/Private/SpeedTreeImportFactory.cpp
--- a/Engine/Plugins/Editor/SpeedTreeImporter/Source/SpeedTreeImporter/Private/SpeedTreeImportFactory.cpp
+++ b/Engine/Plugins/Editor/SpeedTreeImporter/Source/SpeedTreeImporter/Private/SpeedTreeImportFactory.cpp
@@ -916,7 +916,7 @@ static void MakeBodyFromCollisionObjects(UStaticMesh* StaticMesh, const SpeedTre
 		{
 			// sphere object
 			FKSphereElem SphereElem;
-			FMemory::MemZero(SphereElem);
+			FMemory::Memzero(SphereElem);
 			SphereElem.Radius = CollisionObject.m_fRadius;
 			SphereElem.Center = Pos1;
 			AggGeo.SphereElems.Add(SphereElem);
@@ -925,7 +925,7 @@ static void MakeBodyFromCollisionObjects(UStaticMesh* StaticMesh, const SpeedTre
 		{
 			// capsule/sphyll object
 			FKSphylElem SphylElem;
-			FMemory::MemZero(SphylElem);
+			FMemory::Memzero(SphylElem);
 			SphylElem.Radius = CollisionObject.m_fRadius;
 			FVector UpDir = Pos2 - Pos1;
 			SphylElem.Length = UpDir.Size();
diff --git a/Engine/Plugins/Runtime/GearVR/Source/GearVR/Private/GearVR.cpp b/Engine/Plugins/Runtime/GearVR/Source/GearVR/Private/GearVR.cpp
--- a/Engine/Plugins/Runtime/GearVR/Source/GearVR/Private/GearVR.cpp
+++ b/Engine/Plugins/Runtime/GearVR/Source/GearVR/Private/GearVR.cpp
@@ -989,7 +989,7 @@ void FGearVR::Startup()
 
 	UE_LOG(LogHMD, Log, TEXT("GearVR starting with CPU: %d GPU: %d"), CpuLevel, GpuLevel);
 
-	FMemory::MemZero(VrModeParms);
+	FMemory::Memzero(VrModeParms);
 	VrModeParms.AsynchronousTimeWarp = true;
 	VrModeParms.DistortionFileName = NULL;
 	VrModeParms.EnableImageServer = false;
diff --git a/Engine/Plugins/Runtime/TwitchLiveStreaming/Source/TwitchLiveStreaming/TwitchLiveStreaming.cpp b/Engine/Plugins/Runtime/TwitchLiveStreaming/Source/TwitchLiveStreaming/TwitchLiveStreaming.cpp
--- a/Engine/Plugins/Runtime/TwitchLiveStreaming/Source/TwitchLiveStreaming/TwitchLiveStreaming.cpp
+++ b/Engine/Plugins/Runtime/TwitchLiveStreaming/Source/TwitchLiveStreaming/TwitchLiveStreaming.cpp
@@ -72,7 +72,7 @@ void FTwitchLiveStreaming::StartupModule()
 	bWantsChatEnabled = false;
 	ChatState = EChatState::Uninitialized;
 
-	FMemory::MemZero( VideoBuffers );
+	FMemory::Memzero( VideoBuffers );
 
 	// @todo twitch: Sometimes crashes on exit after streaming in the editor on Windows 8.1
 	// @todo twitch urgent: Authorization of end-user games is having problems (bad auth token error)
@@ -436,7 +436,7 @@ void FTwitchLiveStreaming::Async_AuthenticateWithTwitchDirectly( const FString&
 	check( TwitchState == ETwitchState::ReadyToAuthenticate );
 
 	// This will be filled in by TTV_RequestAuthToken()
-	FMemory::MemZero( TwitchAuthToken );
+	FMemory::Memzero( TwitchAuthToken );
 
 	TTV_AuthParams TwitchAuthParams;
 	TwitchAuthParams.size = sizeof( TwitchAuthParams );
@@ -617,7 +617,7 @@ void FTwitchLiveStreaming::Async_LoginToTwitch()
 	check( TwitchState == ETwitchState::ReadyToLogin );
 
 	// Calling TTV_Login() will fill the TwitchChannelInfo structure with information about the Twitch channel
-	FMemory::MemZero( TwitchChannelInfo );	// @todo twitch: We should expose this to the live stream API
+	FMemory::Memzero( TwitchChannelInfo );	// @todo twitch: We should expose this to the live stream API
 	TwitchChannelInfo.size = sizeof( TwitchChannelInfo );
 
 	// Called when TwitchLogin finishes async work
@@ -670,7 +670,7 @@ void FTwitchLiveStreaming::Async_GetIngestServers()
 	check( BroadcastState == EBroadcastState::Idle );
 
 	// This will be filled in by calling TTV_GetIngestServers()
-	FMemory::MemZero( TwitchIngestList );
+	FMemory::Memzero( TwitchIngestList );
 
 	// Called when TwitchGetIngestServers finishes async work
 	void* CallbackUserDataPtr = this;
@@ -722,7 +722,7 @@ void FTwitchLiveStreaming::Async_StartBroadcasting()
 
 	// Pick an ingest server
 	TTV_IngestServer SelectedTwitchIngestServer;
-	FMemory::MemZero( SelectedTwitchIngestServer );
+	FMemory::Memzero( SelectedTwitchIngestServer );
 	{
 		// Select the default ingest server recommended by the Twitch SDK
 		// @todo twitch: We might want to run an 'ingest test' first.  This will help with bitrate selection too, apparently
@@ -794,7 +794,7 @@ void FTwitchLiveStreaming::Async_StartBroadcasting()
 	{
 		// @todo twitch: Make sure these settings are all filled in!
 		TTV_VideoParams TwitchVideoParams;
-		FMemory::MemZero( TwitchVideoParams );
+		FMemory::Memzero( TwitchVideoParams );
 		TwitchVideoParams.size = sizeof( TwitchVideoParams );
 		TwitchVideoParams.outputWidth = BroadcastConfig.VideoBufferWidth;
 		TwitchVideoParams.outputHeight = BroadcastConfig.VideoBufferHeight;
@@ -825,7 +825,7 @@ void FTwitchLiveStreaming::Async_StartBroadcasting()
 		}
 
 		TTV_AudioParams TwitchAudioParams;
-		FMemory::MemZero( TwitchAudioParams );
+		FMemory::Memzero( TwitchAudioParams );
 		TwitchAudioParams.size = sizeof( TwitchAudioParams );
 		TwitchAudioParams.audioEnabled = BroadcastConfig.bCaptureAudioFromComputer || BroadcastConfig.bCaptureAudioFromMicrophone;
 		TwitchAudioParams.enableMicCapture = BroadcastConfig.bCaptureAudioFromMicrophone;
@@ -1989,7 +1989,7 @@ void FTwitchLiveStreaming::Async_GetGameLiveStreams( const FString& GameName, FQ
 	auto* CallbackPayload = new FCallbackPayload();	// This will be deleted by the async callback function below
 	CallbackPayload->This = this;
 	CallbackPayload->GameName = GameName;
-	FMemory::MemZero( CallbackPayload->TwitchGameStreamList );
+	FMemory::Memzero( CallbackPayload->TwitchGameStreamList );
 	CallbackPayload->CompletionCallback = CompletionCallback;
 
 	void* CallbackUserDataPtr = CallbackPayload;
diff --git a/Engine/Source/Developer/Profiler/Private/Widgets/SDataGraph.cpp b/Engine/Source/Developer/Profiler/Private/Widgets/SDataGraph.cpp
--- a/Engine/Source/Developer/Profiler/Private/Widgets/SDataGraph.cpp
+++ b/Engine/Source/Developer/Profiler/Private/Widgets/SDataGraph.cpp
@@ -197,8 +197,8 @@ SDataGraph::SDataGraph()
 {
 	ScaleY = FMath::Pow( 2.0f, MouseWheelAcc );
 
-	FMemory::MemZero( FrameIndices );
-	FMemory::MemZero( FrameTimesMS );
+	FMemory::Memzero( FrameIndices );
+	FMemory::Memzero( FrameTimesMS );
 }
 
 SDataGraph::~SDataGraph()
diff --git a/Engine/Source/Developer/Profiler/Private/Widgets/SProfilerMiniView.cpp b/Engine/Source/Developer/Profiler/Private/Widgets/SProfilerMiniView.cpp
--- a/Engine/Source/Developer/Profiler/Private/Widgets/SProfilerMiniView.cpp
+++ b/Engine/Source/Developer/Profiler/Private/Widgets/SProfilerMiniView.cpp
@@ -30,7 +30,7 @@ void SProfilerMiniView::Reset()
 	RecentlyAddedFrames.Empty();
 	StatMetadata = nullptr;
 
-	FMemory::MemZero( PaintStateMemory );
+	FMemory::Memzero( PaintStateMemory );
 	PaintState = nullptr;
 
 	MousePositionOnButtonDown = FVector2D::ZeroVector;
diff --git a/Engine/Source/Developer/Profiler/Private/Widgets/SProfilerThreadView.cpp b/Engine/Source/Developer/Profiler/Private/Widgets/SProfilerThreadView.cpp
--- a/Engine/Source/Developer/Profiler/Private/Widgets/SProfilerThreadView.cpp
+++ b/Engine/Source/Developer/Profiler/Private/Widgets/SProfilerThreadView.cpp
@@ -26,7 +26,7 @@ void SProfilerThreadView::Reset()
 {
 	ProfilerStream = nullptr;
 
-	FMemory::MemZero( PaintStateMemory );
+	FMemory::Memzero( PaintStateMemory );
 	PaintState = nullptr;
 
 	MousePosition = FVector2D::ZeroVector;
diff --git a/Engine/Source/Developer/ShaderFormatOpenGL/Private/OpenGLShaderCompiler.cpp b/Engine/Source/Developer/ShaderFormatOpenGL/Private/OpenGLShaderCompiler.cpp
--- a/Engine/Source/Developer/ShaderFormatOpenGL/Private/OpenGLShaderCompiler.cpp
+++ b/Engine/Source/Developer/ShaderFormatOpenGL/Private/OpenGLShaderCompiler.cpp
@@ -117,7 +117,7 @@ static void PlatformInitPixelFormatForDevice(HDC DeviceContext)
 {
 	// Pixel format descriptor for the context.
 	PIXELFORMATDESCRIPTOR PixelFormatDesc;
-	FMemory::MemZero(PixelFormatDesc);
+	FMemory::Memzero(PixelFormatDesc);
 	PixelFormatDesc.nSize		= sizeof(PIXELFORMATDESCRIPTOR);
 	PixelFormatDesc.nVersion	= 1;
 	PixelFormatDesc.dwFlags		= PFD_DRAW_TO_WINDOW | PFD_SUPPORT_OPENGL | PFD_DOUBLEBUFFER;
@@ -149,7 +149,7 @@ static void PlatformCreateDummyGLWindow(FPlatformOpenGLContext* OutContext)
 		WNDCLASS wc;
 
 		bInitializedWindowClass = true;
-		FMemory::MemZero(wc);
+		FMemory::Memzero(wc);
 		wc.style = CS_OWNDC;
 		wc.lpfnWndProc = PlatformDummyGLWndproc;
 		wc.cbClsExtra = 0;
diff --git a/Engine/Source/Developer/StandaloneRenderer/Private/Windows/OpenGL/SlateOpenGLContext.cpp b/Engine/Source/Developer/StandaloneRenderer/Private/Windows/OpenGL/SlateOpenGLContext.cpp
--- a/Engine/Source/Developer/StandaloneRenderer/Private/Windows/OpenGL/SlateOpenGLContext.cpp
+++ b/Engine/Source/Developer/StandaloneRenderer/Private/Windows/OpenGL/SlateOpenGLContext.cpp
@@ -26,7 +26,7 @@ static HWND CreateDummyGLWindow()
 		WNDCLASS wc;
 
 		bInitializedWindowClass = true;
-		FMemory::MemZero(wc);
+		FMemory::Memzero(wc);
 		wc.style = CS_OWNDC;
 		wc.lpfnWndProc = DummyGLWndproc;
 		wc.cbClsExtra = 0;
diff --git a/Engine/Source/Developer/TextureFormatAndroid/Private/TextureFormatAndroid.cpp b/Engine/Source/Developer/TextureFormatAndroid/Private/TextureFormatAndroid.cpp
--- a/Engine/Source/Developer/TextureFormatAndroid/Private/TextureFormatAndroid.cpp
+++ b/Engine/Source/Developer/TextureFormatAndroid/Private/TextureFormatAndroid.cpp
@@ -70,8 +70,9 @@ static bool CompressImageUsingQonvert(
 
 	TQonvertImage SrcImg;
 	TQonvertImage DstImg;
-	FMemory::MemZero(SrcImg);
-	FMemory::MemZero(DstImg);
+
+	FMemory::Memzero(SrcImg);
+	FMemory::Memzero(DstImg);
 
 	SrcImg.nWidth    = SizeX;
 	SrcImg.nHeight   = SizeY;
diff --git a/Engine/Source/Editor/UnrealEd/Private/EditorViewportClient.cpp b/Engine/Source/Editor/UnrealEd/Private/EditorViewportClient.cpp
--- a/Engine/Source/Editor/UnrealEd/Private/EditorViewportClient.cpp
+++ b/Engine/Source/Editor/UnrealEd/Private/EditorViewportClient.cpp
@@ -4487,7 +4487,7 @@ void FEditorViewportStats::Initialize()
 	if ( !bInitialized )
 	{
 		bInitialized = true;
-		FMemory::MemZero(DataPoints);
+		FMemory::Memzero(DataPoints);
 	}
 }
 
@@ -4569,7 +4569,7 @@ void FEditorViewportStats::SendUsageData()
 	FEngineAnalytics::GetProvider().RecordEvent(FString("Editor.Usage.Viewport.Orthographic"), OrthographicUsage);
 
 	// Clear all the usage data in case we do it twice.
-	FMemory::MemZero(DataPoints);
+	FMemory::Memzero(DataPoints);
 }
 
 
diff --git a/Engine/Source/Editor/UnrealEd/Private/Factories/HDRLoader.cpp b/Engine/Source/Editor/UnrealEd/Private/Factories/HDRLoader.cpp
--- a/Engine/Source/Editor/UnrealEd/Private/Factories/HDRLoader.cpp
+++ b/Engine/Source/Editor/UnrealEd/Private/Factories/HDRLoader.cpp
@@ -222,7 +222,7 @@ void FHDRLoadHelper::ExtractDDSInRGBE(TArray<uint8>& OutDDSFile) const
 
 		FDDSFileHeader* header = (FDDSFileHeader*)(OutDDSFile.GetData() + 4);
 
-		FMemory::MemZero(*header);
+		FMemory::Memzero(*header);
 
 		header->dwSize = sizeof(FDDSFileHeader);
 		header->dwFlags = DDSF_Caps | DDSF_Height | DDSF_Width | DDSF_PixelFormat;
diff --git a/Engine/Source/Editor/UnrealEd/Private/Fbx/FbxMainImport.cpp b/Engine/Source/Editor/UnrealEd/Private/Fbx/FbxMainImport.cpp
--- a/Engine/Source/Editor/UnrealEd/Private/Fbx/FbxMainImport.cpp
+++ b/Engine/Source/Editor/UnrealEd/Private/Fbx/FbxMainImport.cpp
@@ -292,7 +292,7 @@ FFbxImporter::FFbxImporter()
 	Scene = NULL;
 	
 	ImportOptions = new FBXImportOptions();
-	FMemory::MemZero(*ImportOptions);
+	FMemory::Memzero(*ImportOptions);
 	
 	CurPhase = NOTSTARTED;
 }
diff --git a/Engine/Source/Editor/UnrealEd/Private/Lightmass/Lightmass.cpp b/Engine/Source/Editor/UnrealEd/Private/Lightmass/Lightmass.cpp
--- a/Engine/Source/Editor/UnrealEd/Private/Lightmass/Lightmass.cpp
+++ b/Engine/Source/Editor/UnrealEd/Private/Lightmass/Lightmass.cpp
@@ -106,7 +106,7 @@ static NSwarm::TChannelFlags LM_DEBUGOUTPUT_CHANNEL_FLAGS		= NSwarm::SWARM_JOB_C
 
 void Copy( const ULightComponentBase* In, Lightmass::FLightData& Out )
 {	
-	FMemory::MemZero(Out);
+	FMemory::Memzero(Out);
 
 	Out.LightFlags = 0;
 	if (In->CastShadows)
diff --git a/Engine/Source/Runtime/AIModule/Private/AITypes.cpp b/Engine/Source/Runtime/AIModule/Private/AITypes.cpp
--- a/Engine/Source/Runtime/AIModule/Private/AITypes.cpp
+++ b/Engine/Source/Runtime/AIModule/Private/AITypes.cpp
@@ -13,7 +13,7 @@ FAIResourceLock::FAIResourceLock()
 
 void FAIResourceLock::ForceClearAllLocks()
 {
-	FMemory::MemZero(Locks);
+	FMemory::Memzero(Locks);
 }
 
 FString FAIResourceLock::GetLockPriorityName() const
diff --git a/Engine/Source/Runtime/Core/Private/GenericPlatform/GenericPlatformStackWalk.cpp b/Engine/Source/Runtime/Core/Private/GenericPlatform/GenericPlatformStackWalk.cpp
--- a/Engine/Source/Runtime/Core/Private/GenericPlatform/GenericPlatformStackWalk.cpp
+++ b/Engine/Source/Runtime/Core/Private/GenericPlatform/GenericPlatformStackWalk.cpp
@@ -52,7 +52,7 @@ void FGenericPlatformStackWalk::StackWalkAndDump( ANSICHAR* HumanReadableString,
 	// Temporary memory holding the stack trace.
 	static const int MAX_DEPTH = 100;
 	uint64 StackTrace[MAX_DEPTH];
-	FMemory::MemZero( StackTrace );
+	FMemory::Memzero( StackTrace );
 
 	// Capture stack backtrace.
 	FPlatformStackWalk::CaptureStackBackTrace( StackTrace, MAX_DEPTH, Context );
diff --git a/Engine/Source/Runtime/Core/Private/Mac/CocoaThread.cpp b/Engine/Source/Runtime/Core/Private/Mac/CocoaThread.cpp
--- a/Engine/Source/Runtime/Core/Private/Mac/CocoaThread.cpp
+++ b/Engine/Source/Runtime/Core/Private/Mac/CocoaThread.cpp
@@ -168,7 +168,7 @@ private:
 		{
 			CFRunLoopSourceContext Context;
 			FCocoaRunLoopSourceInfo* Info = [[FCocoaRunLoopSourceInfo alloc] initWithSource:this];
-			FMemory::MemZero(Context);
+			FMemory::Memzero(Context);
 			Context.version = 0;
 			Context.info = Info;
 			Context.retain = CFRetain;
diff --git a/Engine/Source/Runtime/Core/Private/Mac/MacApplication.cpp b/Engine/Source/Runtime/Core/Private/Mac/MacApplication.cpp
--- a/Engine/Source/Runtime/Core/Private/Mac/MacApplication.cpp
+++ b/Engine/Source/Runtime/Core/Private/Mac/MacApplication.cpp
@@ -230,7 +230,7 @@ FMacApplication::FMacApplication()
 								   }];
 
 #if WITH_EDITOR
-	FMemory::MemZero(GestureUsage);
+	FMemory::Memzero(GestureUsage);
 	LastGestureUsed = EGestureEvent::None;
 #endif
 }
@@ -1368,7 +1368,7 @@ void FMacApplication::SendAnalytics(IAnalyticsProvider* Provider)
 
     Provider->RecordEvent(FString("Mac.Gesture.Usage"), GestureAttributes);
 
-	FMemory::MemZero(GestureUsage);
+	FMemory::Memzero(GestureUsage);
 	LastGestureUsed = EGestureEvent::None;
 }
 
diff --git a/Engine/Source/Runtime/Core/Private/Mac/MacPlatformMisc.cpp b/Engine/Source/Runtime/Core/Private/Mac/MacPlatformMisc.cpp
--- a/Engine/Source/Runtime/Core/Private/Mac/MacPlatformMisc.cpp
+++ b/Engine/Source/Runtime/Core/Private/Mac/MacPlatformMisc.cpp
@@ -1487,7 +1487,7 @@ void FMacCrashContext::GenerateMinidump(char const* Path) const
 		
 		static const int MAX_DEPTH = 100;
 		static uint64 StackTrace[MAX_DEPTH];
-		FMemory::MemZero( StackTrace );
+		FMemory::Memzero( StackTrace );
 		FPlatformStackWalk::CaptureStackBackTrace( StackTrace, MAX_DEPTH, Context );
 		
 		// Skip the first two entries as they are inside the stack walking code.
diff --git a/Engine/Source/Runtime/Core/Private/Windows/WindowsApplication.cpp b/Engine/Source/Runtime/Core/Private/Windows/WindowsApplication.cpp
--- a/Engine/Source/Runtime/Core/Private/Windows/WindowsApplication.cpp
+++ b/Engine/Source/Runtime/Core/Private/Windows/WindowsApplication.cpp
@@ -84,9 +84,9 @@ FWindowsApplication::FWindowsApplication( const HINSTANCE HInstance, const HICON
 	// Save the current sticky/toggle/filter key settings so they can be restored them later
 	// If there are .ini settings, use them instead of the current system settings.
 	// NOTE: Whenever we exit and restore these settings gracefully, the .ini settings are removed.
-	FMemory::MemZero(StartupStickyKeys);
-	FMemory::MemZero(StartupToggleKeys);
-	FMemory::MemZero(StartupFilterKeys);
+	FMemory::Memzero(StartupStickyKeys);
+	FMemory::Memzero(StartupToggleKeys);
+	FMemory::Memzero(StartupFilterKeys);
 	
 	StartupStickyKeys.cbSize = sizeof(StartupStickyKeys);
 	StartupToggleKeys.cbSize = sizeof(StartupToggleKeys);
@@ -590,7 +590,7 @@ void FWindowsApplication::GetInitialDisplayMetrics( FDisplayMetrics& OutDisplayM
 EWindowTitleAlignment::Type FWindowsApplication::GetWindowTitleAlignment() const
 {
 	OSVERSIONINFOEX VersionInfo;
-	FMemory::MemZero(VersionInfo);
+	FMemory::Memzero(VersionInfo);
 	VersionInfo.dwMajorVersion = 6;
 	VersionInfo.dwMinorVersion = 2;
 	VersionInfo.dwOSVersionInfoSize = sizeof(VersionInfo);
diff --git a/Engine/Source/Runtime/Core/Private/Windows/WindowsPlatformMisc.cpp b/Engine/Source/Runtime/Core/Private/Windows/WindowsPlatformMisc.cpp
--- a/Engine/Source/Runtime/Core/Private/Windows/WindowsPlatformMisc.cpp
+++ b/Engine/Source/Runtime/Core/Private/Windows/WindowsPlatformMisc.cpp
@@ -369,7 +369,7 @@ static void SetProcessMemoryLimit( SIZE_T ProcessMemoryLimitMB )
 {
 	HANDLE JobObject = ::CreateJobObject(nullptr, TEXT("UE4-JobObject"));
 	JOBOBJECT_EXTENDED_LIMIT_INFORMATION JobLimitInfo;
-	FMemory::MemZero( JobLimitInfo );
+	FMemory::Memzero( JobLimitInfo );
 	JobLimitInfo.ProcessMemoryLimit = 1024*1024*ProcessMemoryLimitMB;
 	JobLimitInfo.BasicLimitInformation.LimitFlags = JOB_OBJECT_LIMIT_PROCESS_MEMORY;
 	const BOOL bSetJob = ::SetInformationJobObject(JobObject,JobObjectExtendedLimitInformation,&JobLimitInfo,sizeof(JobLimitInfo));
@@ -2181,7 +2181,7 @@ FString FWindowsPlatformMisc::GetPrimaryGPUBrand()
 				break;
 			}
 
-			FMemory::MemZero( DisplayDevice );
+			FMemory::Memzero( DisplayDevice );
 			DisplayDevice.cb = sizeof( DisplayDevice );
 			DeviceIndex++;
 		}
diff --git a/Engine/Source/Runtime/Core/Private/Windows/WindowsPlatformNamedPipe.cpp b/Engine/Source/Runtime/Core/Private/Windows/WindowsPlatformNamedPipe.cpp
--- a/Engine/Source/Runtime/Core/Private/Windows/WindowsPlatformNamedPipe.cpp
+++ b/Engine/Source/Runtime/Core/Private/Windows/WindowsPlatformNamedPipe.cpp
@@ -24,7 +24,7 @@ FWindowsPlatformNamedPipe::FWindowsPlatformNamedPipe() :
 	bIsServer(false),
 	State(FPlatformNamedPipe::State_Uninitialized)
 {
-	FMemory::MemZero(Overlapped);
+	FMemory::Memzero(Overlapped);
 }
 
 FWindowsPlatformNamedPipe::~FWindowsPlatformNamedPipe()
@@ -71,7 +71,7 @@ bool FWindowsPlatformNamedPipe::Create(const FString& PipeName, bool bAsServer,
 		return false;
 	}
 
-	FMemory::MemZero(Overlapped);
+	FMemory::Memzero(Overlapped);
 
 	State = bAsServer ? State_Created : State_ReadyForRW;
 
diff --git a/Engine/Source/Runtime/Core/Private/Windows/WindowsPlatformStackWalk.cpp b/Engine/Source/Runtime/Core/Private/Windows/WindowsPlatformStackWalk.cpp
--- a/Engine/Source/Runtime/Core/Private/Windows/WindowsPlatformStackWalk.cpp
+++ b/Engine/Source/Runtime/Core/Private/Windows/WindowsPlatformStackWalk.cpp
@@ -67,7 +67,7 @@ static int32 CaptureStackTraceHelper( uint64 *BackTrace, uint32 MaxDepth, CONTEX
 		ThreadHandle	= GetCurrentThread();
 
 		// Zero out stack frame.
-		FMemory::MemZero( StackFrame64 );
+		FMemory::Memzero( StackFrame64 );
 
 		// Initialize the STACKFRAME structure.
 		StackFrame64.AddrPC.Mode         = AddrModeFlat;
diff --git a/Engine/Source/Runtime/Core/Private/Windows/WindowsWindow.cpp b/Engine/Source/Runtime/Core/Private/Windows/WindowsWindow.cpp
--- a/Engine/Source/Runtime/Core/Private/Windows/WindowsWindow.cpp
+++ b/Engine/Source/Runtime/Core/Private/Windows/WindowsWindow.cpp
@@ -213,10 +213,10 @@ FWindowsWindow::FWindowsWindow()
 	, OLEReferenceCount(0)
 	, bIsVisible( false )
 {
-	FMemory::MemZero(PreFullscreenWindowPlacement);
+	FMemory::Memzero(PreFullscreenWindowPlacement);
 	PreFullscreenWindowPlacement.length = sizeof(WINDOWPLACEMENT);
 
-	FMemory::MemZero(PreParentMinimizedWindowPlacement);
+	FMemory::Memzero(PreParentMinimizedWindowPlacement);
 	PreParentMinimizedWindowPlacement.length = sizeof(WINDOWPLACEMENT);
 }
 
@@ -311,7 +311,7 @@ void FWindowsWindow::AdjustWindowRegion( int32 Width, int32 Height )
 void FWindowsWindow::ReshapeWindow( int32 NewX, int32 NewY, int32 NewWidth, int32 NewHeight )
 {
 	WINDOWINFO WindowInfo;
-	FMemory::MemZero( WindowInfo );
+	FMemory::Memzero( WindowInfo );
 	WindowInfo.cbSize = sizeof( WindowInfo );
 	::GetWindowInfo( HWnd, &WindowInfo );
 
@@ -700,7 +700,7 @@ bool FWindowsWindow::IsPointInWindow( int32 X, int32 Y ) const
 int32 FWindowsWindow::GetWindowBorderSize() const
 {
 	WINDOWINFO WindowInfo;
-	FMemory::MemZero( WindowInfo );
+	FMemory::Memzero( WindowInfo );
 	WindowInfo.cbSize = sizeof( WindowInfo );
 	::GetWindowInfo( HWnd, &WindowInfo );
 
diff --git a/Engine/Source/Runtime/Core/Public/HAL/UnrealMemory.h b/Engine/Source/Runtime/Core/Public/HAL/UnrealMemory.h
--- a/Engine/Source/Runtime/Core/Public/HAL/UnrealMemory.h
+++ b/Engine/Source/Runtime/Core/Public/HAL/UnrealMemory.h
@@ -38,17 +38,37 @@ struct CORE_API FMemory
 		return FPlatformMemory::Memset( Dest, Char, Count );
 	}
 
+	template< class T > 
+	static FORCEINLINE void Memset( T& Src, uint8 ValueToSet )
+	{
+		static_assert( !TIsPointerType<T>::Value, "For pointers use the three parameters function");
+		Memset( &Src, ValueToSet, sizeof( T ) );
+	}
+
 	static FORCEINLINE void* Memzero(void* Dest, SIZE_T Count)
 	{
 		return FPlatformMemory::Memzero( Dest, Count );
 	}
 
+	template< class T > 
+	static FORCEINLINE void Memzero( T& Src )
+	{
+		static_assert( !TIsPointerType<T>::Value, "For pointers use the two parameters function");
+		Memzero( &Src, sizeof( T ) );
+	}
 
 	static FORCEINLINE void* Memcpy(void* Dest, const void* Src, SIZE_T Count)
 	{
 		return FPlatformMemory::Memcpy(Dest,Src,Count);
 	}
 
+	template< class T > 
+	static FORCEINLINE void Memcpy( T& Dest, const T& Src )
+	{
+		static_assert( !TIsPointerType<T>::Value, "For pointers use the three parameters function");
+		Memcpy( &Dest, &Src, sizeof( T ) );
+	}
+
 	static FORCEINLINE void* BigBlockMemcpy(void* Dest, const void* Src, SIZE_T Count)
 	{
 		return FPlatformMemory::BigBlockMemcpy(Dest,Src,Count);
@@ -64,22 +84,27 @@ struct CORE_API FMemory
 		FPlatformMemory::Memswap(Ptr1,Ptr2,Size);
 	}
 
-	template< class T > 
+	template< class T >
+	DEPRECATED(4.8, "Please use Memset.")
 	static FORCEINLINE void MemSet( T& Src, uint8 ValueToSet )
 	{
+		static_assert( !TIsPointerType<T>::Value, "For pointers use the three parameters function");
 		FMemory::Memset( &Src, ValueToSet, sizeof( T ) );
 	}
 
 	template< class T > 
+	DEPRECATED(4.8, "Please use Memzero.")
 	static FORCEINLINE void MemZero( T& Src )
 	{
+		static_assert( !TIsPointerType<T>::Value, "For pointers use the two parameters function");
 		FMemory::Memset( &Src, 0, sizeof( T ) );
 	}
 
-
-	template< class T > 
+	template< class T >
+	DEPRECATED(4.8, "Please use Memcpy.")
 	static FORCEINLINE void MemCopy( T& Dest, const T& Src )
 	{
+		static_assert( !TIsPointerType<T>::Value, "For pointers use the three parameters function");
 		FMemory::Memcpy( &Dest, &Src, sizeof( T ) );
 	}
 
diff --git a/Engine/Source/Runtime/Core/Public/UObject/NameTypes.h b/Engine/Source/Runtime/Core/Public/UObject/NameTypes.h
--- a/Engine/Source/Runtime/Core/Public/UObject/NameTypes.h
+++ b/Engine/Source/Runtime/Core/Public/UObject/NameTypes.h
@@ -319,7 +319,7 @@ public:
 		: NumElements(0)
 		, NumChunks(0)
 	{
-		FMemory::MemZero(Chunks);
+		FMemory::Memzero(Chunks);
 	}
 	/** 
 	 * Return the number of elements in the array 
diff --git a/Engine/Source/Runtime/Engine/Classes/AI/Navigation/RecastNavMesh.h b/Engine/Source/Runtime/Engine/Classes/AI/Navigation/RecastNavMesh.h
--- a/Engine/Source/Runtime/Engine/Classes/AI/Navigation/RecastNavMesh.h
+++ b/Engine/Source/Runtime/Engine/Classes/AI/Navigation/RecastNavMesh.h
@@ -599,8 +599,8 @@ public:
 			, HitTime(FLT_MAX)
 			, HitNormal(0.f)
 		{
-			FMemory::MemZero(CorridorPolys);
-			FMemory::MemZero(CorridorCost);
+			FMemory::Memzero(CorridorPolys);
+			FMemory::Memzero(CorridorCost);
 		}
 
 		FORCEINLINE int32 GetMaxCorridorSize() const { return MAX_PATH_CORRIDOR_POLYS; }
diff --git a/Engine/Source/Runtime/Engine/Private/AI/Navigation/RecastNavMeshGenerator.cpp b/Engine/Source/Runtime/Engine/Private/AI/Navigation/RecastNavMeshGenerator.cpp
--- a/Engine/Source/Runtime/Engine/Private/AI/Navigation/RecastNavMeshGenerator.cpp
+++ b/Engine/Source/Runtime/Engine/Private/AI/Navigation/RecastNavMeshGenerator.cpp
@@ -1167,7 +1167,7 @@ struct FOffMeshData
 			}
 
 			dtOffMeshLinkCreateParams NewInfo;
-			FMemory::MemZero(NewInfo);
+			FMemory::Memzero(NewInfo);
 
 			// not doing anything to link's points order - should be already ordered properly by link processor
 			StoreUnrealPoint(NewInfo.vertsA0, LocalToWorld.TransformPosition(Link.Left));
@@ -1205,7 +1205,7 @@ struct FOffMeshData
 			}
 
 			dtOffMeshLinkCreateParams NewInfo;
-			FMemory::MemZero(NewInfo);
+			FMemory::Memzero(NewInfo);
 
 			// not doing anything to link's points order - should be already ordered properly by link processor
 			StoreUnrealPoint(NewInfo.vertsA0, LocalToWorld.TransformPosition(Link.LeftStart));
@@ -2865,7 +2865,7 @@ bool FRecastNavMeshGenerator::ConstructTiledNavMesh()
 		++Version;
 		
 		dtNavMeshParams TiledMeshParameters;
-		FMemory::MemZero(TiledMeshParameters);	
+		FMemory::Memzero(TiledMeshParameters);	
 		rcVcopy(TiledMeshParameters.orig, &FVector::ZeroVector.X);
 		TiledMeshParameters.tileWidth = Config.tileSize * Config.cs;
 		TiledMeshParameters.tileHeight = Config.tileSize * Config.cs;
diff --git a/Engine/Source/Runtime/Engine/Public/AI/Navigation/RecastNavMeshGenerator.h b/Engine/Source/Runtime/Engine/Public/AI/Navigation/RecastNavMeshGenerator.h
--- a/Engine/Source/Runtime/Engine/Public/AI/Navigation/RecastNavMeshGenerator.h
+++ b/Engine/Source/Runtime/Engine/Public/AI/Navigation/RecastNavMeshGenerator.h
@@ -54,7 +54,7 @@ struct FRecastBuildConfig : public rcConfig
 
 	void Reset()
 	{
-		FMemory::MemZero(*this);
+		FMemory::Memzero(*this);
 		bPerformVoxelFiltering = true;
 		bGenerateDetailedMesh = true;
 		bGenerateBVTree = true;
diff --git a/Engine/Source/Runtime/Engine/Public/StaticMeshResources.h b/Engine/Source/Runtime/Engine/Public/StaticMeshResources.h
--- a/Engine/Source/Runtime/Engine/Public/StaticMeshResources.h
+++ b/Engine/Source/Runtime/Engine/Public/StaticMeshResources.h
@@ -27,7 +27,7 @@ public:
 		, BasePercentTrianglesMult(1.0f)
 		, DisplayName( NSLOCTEXT( "UnrealEd", "None", "None" ) )
 	{
-		FMemory::MemZero(SettingsBias);
+		FMemory::Memzero(SettingsBias);
 		SettingsBias.PercentTriangles = 1.0f;
 	}
 
diff --git a/Engine/Source/Runtime/IOS/MetalRHI/Private/MetalCommands.cpp b/Engine/Source/Runtime/IOS/MetalRHI/Private/MetalCommands.cpp
--- a/Engine/Source/Runtime/IOS/MetalRHI/Private/MetalCommands.cpp
+++ b/Engine/Source/Runtime/IOS/MetalRHI/Private/MetalCommands.cpp
@@ -771,7 +771,7 @@ void FMetalDynamicRHI::RHIClearMRT(bool bClearColor,int32 NumClearColors,const F
 			FRHICommandList_RecursiveHazardous RHICmdList;
 			SetGlobalBoundShaderState(RHICmdList, GMaxRHIFeatureLevel, GClearMRTBoundShaderState[0], GVector4VertexDeclaration.VertexDeclarationRHI, *VertexShader, PixelShader);
 			FLinearColor ShaderClearColors[MaxSimultaneousRenderTargets];
-			FMemory::MemZero(ShaderClearColors);
+			FMemory::Memzero(ShaderClearColors);
 
 			for (int32 i = 0; i < NumClearColors; i++)
 			{
diff --git a/Engine/Source/Runtime/IOS/MetalRHI/Private/MetalShaders.cpp b/Engine/Source/Runtime/IOS/MetalRHI/Private/MetalShaders.cpp
--- a/Engine/Source/Runtime/IOS/MetalRHI/Private/MetalShaders.cpp
+++ b/Engine/Source/Runtime/IOS/MetalRHI/Private/MetalShaders.cpp
@@ -407,8 +407,8 @@ FMetalShaderParameterCache::~FMetalShaderParameterCache()
 		FMemory::Free(PackedGlobalUniforms[0]);
 	}
 	
-	FMemory::MemZero(PackedUniformsScratch);
-	FMemory::MemZero(PackedGlobalUniforms);
+	FMemory::Memzero(PackedUniformsScratch);
+	FMemory::Memzero(PackedGlobalUniforms);
 	
 	GlobalUniformArraySize = -1;
 }
diff --git a/Engine/Source/Runtime/Linux/SteamController/Private/SteamControllerInterface.cpp b/Engine/Source/Runtime/Linux/SteamController/Private/SteamControllerInterface.cpp
--- a/Engine/Source/Runtime/Linux/SteamController/Private/SteamControllerInterface.cpp
+++ b/Engine/Source/Runtime/Linux/SteamController/Private/SteamControllerInterface.cpp
@@ -31,7 +31,7 @@ SteamControllerInterface::SteamControllerInterface( const TSharedRef< FGenericAp
 		// [RCL] 2014-05-05 FIXME: disable when could not init?
 		if (bInited)
 		{
-			FMemory::MemZero(ControllerStates);
+			FMemory::Memzero(ControllerStates);
 		}
 	}
 	else
diff --git a/Engine/Source/Runtime/Online/HTTP/Private/Windows/HttpWinInet.cpp b/Engine/Source/Runtime/Online/HTTP/Private/Windows/HttpWinInet.cpp
--- a/Engine/Source/Runtime/Online/HTTP/Private/Windows/HttpWinInet.cpp
+++ b/Engine/Source/Runtime/Online/HTTP/Private/Windows/HttpWinInet.cpp
@@ -1213,7 +1213,7 @@ void FURLWinInet::ClearCachedData() const
 {
 	URLPtr = NULL;
 	URLParameters.Reset();
-	FMemory::MemZero(URLParts);
+	FMemory::Memzero(URLParts);
 }
 
 void FURLWinInet::CrackUrlParameters() const
diff --git a/Engine/Source/Runtime/Online/OnlineSubsystemSteam/Private/OnlineAsyncTaskManagerSteam.cpp b/Engine/Source/Runtime/Online/OnlineSubsystemSteam/Private/OnlineAsyncTaskManagerSteam.cpp
--- a/Engine/Source/Runtime/Online/OnlineSubsystemSteam/Private/OnlineAsyncTaskManagerSteam.cpp
+++ b/Engine/Source/Runtime/Online/OnlineSubsystemSteam/Private/OnlineAsyncTaskManagerSteam.cpp
@@ -87,7 +87,7 @@ private:
 	FOnlineAsyncEventSteamLobbyEnter() :
 		FOnlineAsyncEvent(NULL)
 	{
-		FMemory::MemZero(CallbackResults);
+		FMemory::Memzero(CallbackResults);
 	}
 
 public:
@@ -161,7 +161,7 @@ private:
 	FOnlineAsyncEventSteamLobbyChatUpdate() :
 		FOnlineAsyncEvent(NULL)
 	{
-		FMemory::MemZero(CallbackResults);
+		FMemory::Memzero(CallbackResults);
 	}
 
 public:
diff --git a/Engine/Source/Runtime/OpenGLDrv/Private/OpenGLShaders.cpp b/Engine/Source/Runtime/OpenGLDrv/Private/OpenGLShaders.cpp
--- a/Engine/Source/Runtime/OpenGLDrv/Private/OpenGLShaders.cpp
+++ b/Engine/Source/Runtime/OpenGLDrv/Private/OpenGLShaders.cpp
@@ -1932,8 +1932,8 @@ FOpenGLShaderParameterCache::~FOpenGLShaderParameterCache()
 		FMemory::Free(PackedGlobalUniforms[0]);
 	}
 
-	FMemory::MemZero(PackedUniformsScratch);
-	FMemory::MemZero(PackedGlobalUniforms);
+	FMemory::Memzero(PackedUniformsScratch);
+	FMemory::Memzero(PackedGlobalUniforms);
 
 	GlobalUniformArraySize = -1;
 }
diff --git a/Engine/Source/Runtime/OpenGLDrv/Private/Windows/OpenGLWindows.cpp b/Engine/Source/Runtime/OpenGLDrv/Private/Windows/OpenGLWindows.cpp
--- a/Engine/Source/Runtime/OpenGLDrv/Private/Windows/OpenGLWindows.cpp
+++ b/Engine/Source/Runtime/OpenGLDrv/Private/Windows/OpenGLWindows.cpp
@@ -105,7 +105,7 @@ static void PlatformInitPixelFormatForDevice(HDC DeviceContext)
 {
 	// Pixel format descriptor for the context.
 	PIXELFORMATDESCRIPTOR PixelFormatDesc;
-	FMemory::MemZero(PixelFormatDesc);
+	FMemory::Memzero(PixelFormatDesc);
 	PixelFormatDesc.nSize		= sizeof(PIXELFORMATDESCRIPTOR);
 	PixelFormatDesc.nVersion	= 1;
 	PixelFormatDesc.dwFlags		= PFD_DRAW_TO_WINDOW | PFD_SUPPORT_OPENGL | PFD_DOUBLEBUFFER;
@@ -137,7 +137,7 @@ static void PlatformCreateDummyGLWindow(FPlatformOpenGLContext* OutContext)
 		WNDCLASS wc;
 
 		bInitializedWindowClass = true;
-		FMemory::MemZero(wc);
+		FMemory::Memzero(wc);
 		wc.style = CS_OWNDC;
 		wc.lpfnWndProc = PlatformDummyGLWndproc;
 		wc.cbClsExtra = 0;
diff --git a/Engine/Source/Runtime/Renderer/Private/ClearQuad.cpp b/Engine/Source/Runtime/Renderer/Private/ClearQuad.cpp
--- a/Engine/Source/Runtime/Renderer/Private/ClearQuad.cpp
+++ b/Engine/Source/Runtime/Renderer/Private/ClearQuad.cpp
@@ -100,7 +100,7 @@ void DrawClearQuadMRT(FRHICommandList& RHICmdList, ERHIFeatureLevel::Type Featur
 
 	SetGlobalBoundShaderState(RHICmdList, FeatureLevel, GClearMRTBoundShaderState[FMath::Max(NumClearColors - 1, 0)], GetVertexDeclarationFVector4(), *VertexShader, PixelShader);
 	FLinearColor ShaderClearColors[MaxSimultaneousRenderTargets];
-	FMemory::MemZero(ShaderClearColors);
+	FMemory::Memzero(ShaderClearColors);
 
 	for (int32 i = 0; i < NumClearColors; i++)
 	{
diff --git a/Engine/Source/Runtime/Renderer/Private/TranslucentLighting.cpp b/Engine/Source/Runtime/Renderer/Private/TranslucentLighting.cpp
--- a/Engine/Source/Runtime/Renderer/Private/TranslucentLighting.cpp
+++ b/Engine/Source/Runtime/Renderer/Private/TranslucentLighting.cpp
@@ -1082,7 +1082,7 @@ void ClearVolumeTextures(FRHICommandList& RHICmdList, ERHIFeatureLevel::Type Fea
 		GeometryShader->SetParameters(RHICmdList, VolumeBounds);
 
 		FLinearColor ShaderClearColors[MaxSimultaneousRenderTargets];
-		FMemory::MemZero(ShaderClearColors);
+		FMemory::Memzero(ShaderClearColors);
 
 		for (int32 i = 0; i < NumRenderTargets; i++)
 		{
diff --git a/Engine/Source/Runtime/SlateCore/Private/Layout/Geometry.cpp b/Engine/Source/Runtime/SlateCore/Private/Layout/Geometry.cpp
--- a/Engine/Source/Runtime/SlateCore/Private/Layout/Geometry.cpp
+++ b/Engine/Source/Runtime/SlateCore/Private/Layout/Geometry.cpp
@@ -15,7 +15,7 @@ FGeometry& FGeometry::operator=(const FGeometry& RHS)
 	// HACK to allow us to make FGeometry public members immutable to catch misuse.
 	if (this != &RHS)
 	{
-		FMemory::MemCopy(*this, RHS);
+		FMemory::Memcpy(*this, RHS);
 	}
 	return *this;
 }
diff --git a/Engine/Source/Runtime/Sockets/Private/Linux/SocketSubsystemLinux.cpp b/Engine/Source/Runtime/Sockets/Private/Linux/SocketSubsystemLinux.cpp
--- a/Engine/Source/Runtime/Sockets/Private/Linux/SocketSubsystemLinux.cpp
+++ b/Engine/Source/Runtime/Sockets/Private/Linux/SocketSubsystemLinux.cpp
@@ -109,7 +109,7 @@ TSharedRef<FInternetAddr> FSocketSubsystemLinux::GetLocalHostAddr(FOutputDevice&
 		ifreq IfReqs[8];
 		
 		ifconf IfConfig;
-		FMemory::MemZero(IfConfig);
+		FMemory::Memzero(IfConfig);
 		IfConfig.ifc_req = IfReqs;
 		IfConfig.ifc_len = sizeof(IfReqs);
 		
diff --git a/Engine/Source/Runtime/Windows/D3D11RHI/Private/D3D11Commands.cpp b/Engine/Source/Runtime/Windows/D3D11RHI/Private/D3D11Commands.cpp
--- a/Engine/Source/Runtime/Windows/D3D11RHI/Private/D3D11Commands.cpp
+++ b/Engine/Source/Runtime/Windows/D3D11RHI/Private/D3D11Commands.cpp
@@ -1799,7 +1799,7 @@ void FD3D11DynamicRHI::RHIClearMRT(bool bClearColor,int32 NumClearColors,const F
 			FRHICommandList_RecursiveHazardous RHICmdList;
 			SetGlobalBoundShaderState(RHICmdList, GMaxRHIFeatureLevel, GD3D11ClearMRTBoundShaderState[FMath::Max(BoundRenderTargets.GetNumActiveTargets() - 1, 0)], GD3D11Vector4VertexDeclaration.VertexDeclarationRHI, *VertexShader, PixelShader);
 			FLinearColor ShaderClearColors[MaxSimultaneousRenderTargets];
-			FMemory::MemZero(ShaderClearColors);
+			FMemory::Memzero(ShaderClearColors);
 
 			for (int32 i = 0; i < NumClearColors; i++)
 			{
diff --git a/Engine/Source/Runtime/Windows/D3D11RHI/Private/D3D11UAV.cpp b/Engine/Source/Runtime/Windows/D3D11RHI/Private/D3D11UAV.cpp
--- a/Engine/Source/Runtime/Windows/D3D11RHI/Private/D3D11UAV.cpp
+++ b/Engine/Source/Runtime/Windows/D3D11RHI/Private/D3D11UAV.cpp
@@ -155,7 +155,7 @@ FShaderResourceViewRHIRef FD3D11DynamicRHI::RHICreateShaderResourceView(FVertexB
 
 	// Create a Shader Resource View
 	D3D11_SHADER_RESOURCE_VIEW_DESC SRVDesc;
-	FMemory::MemZero(SRVDesc);
+	FMemory::Memzero(SRVDesc);
 	SRVDesc.ViewDimension = D3D11_SRV_DIMENSION_BUFFER;
 	SRVDesc.Buffer.FirstElement = 0;
 
diff --git a/Engine/Source/Runtime/Windows/D3D11RHI/Private/D3D11Viewport.cpp b/Engine/Source/Runtime/Windows/D3D11RHI/Private/D3D11Viewport.cpp
--- a/Engine/Source/Runtime/Windows/D3D11RHI/Private/D3D11Viewport.cpp
+++ b/Engine/Source/Runtime/Windows/D3D11RHI/Private/D3D11Viewport.cpp
@@ -283,7 +283,7 @@ void FD3D11Viewport::PresentWithVsyncDWM()
 
 	// Find out how long since we last flipped and query DWM for timing information.
 	QueryPerformanceCounter(&Cycles);
-	FMemory::MemZero(TimingInfo);
+	FMemory::Memzero(TimingInfo);
 	TimingInfo.cbSize = sizeof(DWM_TIMING_INFO);
 	DwmGetCompositionTimingInfo(WindowHandle, &TimingInfo);
 
