Output TwitchLiveStreaming to a NotForLicensees folder, since it links against a NotForLicensees third party module.

diff --git a/Engine/Plugins/Runtime/TwitchLiveStreaming/Source/TwitchLiveStreaming/TwitchLiveStreaming.Build.cs b/Engine/Plugins/Runtime/TwitchLiveStreaming/Source/TwitchLiveStreaming/TwitchLiveStreaming.Build.cs
--- a/Engine/Plugins/Runtime/TwitchLiveStreaming/Source/TwitchLiveStreaming/TwitchLiveStreaming.Build.cs
+++ b/Engine/Plugins/Runtime/TwitchLiveStreaming/Source/TwitchLiveStreaming/TwitchLiveStreaming.Build.cs
@@ -6,6 +6,8 @@ namespace UnrealBuildTool.Rules
     {
         public TwitchLiveStreaming(TargetInfo Target)
         {
+			BinariesSubFolder = "NotForLicensees";
+
             PrivateDependencyModuleNames.AddRange(
                 new string[]
 				{
@@ -29,7 +31,6 @@ namespace UnrealBuildTool.Rules
 				});
 
  			bool bHaveTwitchSDK =
-				!UnrealBuildTool.BuildingRocket() &&	// @todo: Twitch support disabled in released builds until it is cleared by legal
  				( System.IO.Directory.Exists( System.IO.Path.Combine( UEBuildConfiguration.UEThirdPartySourceDirectory, "Twitch" ) ) &&
  				  System.IO.Directory.Exists( System.IO.Path.Combine( UEBuildConfiguration.UEThirdPartySourceDirectory, "Twitch", "Twitch-6.17" ) ) ) ||
  				( System.IO.Directory.Exists( System.IO.Path.Combine( UEBuildConfiguration.UEThirdPartySourceDirectory, "NotForLicensees" ) ) &&
